Track sort order in Rewind Debugger.  Sorting by name and adding a priority to specific track types

diff --git a/Engine/Plugins/Animation/GameplayInsights/Source/GameplayInsights/Private/AnimCurvesTrack.cpp b/Engine/Plugins/Animation/GameplayInsights/Source/GameplayInsights/Private/AnimCurvesTrack.cpp
--- a/Engine/Plugins/Animation/GameplayInsights/Source/GameplayInsights/Private/AnimCurvesTrack.cpp
+++ b/Engine/Plugins/Animation/GameplayInsights/Source/GameplayInsights/Private/AnimCurvesTrack.cpp
@@ -14,6 +14,7 @@ namespace RewindDebugger
 
 FAnimCurvesTrack::FAnimCurvesTrack(uint64 InObjectId) : ObjectId(InObjectId)
 {
+	SetIsExpanded(false);
 	Icon = FSlateIcon("EditorStyle", "AnimGraph.Attribute.Curves.Icon", "AnimGraph.Attribute.Curves.Icon");
 }
 
diff --git a/Engine/Plugins/Animation/GameplayInsights/Source/RewindDebugger/Private/RewindDebuggerObjectTrack.cpp b/Engine/Plugins/Animation/GameplayInsights/Source/RewindDebugger/Private/RewindDebuggerObjectTrack.cpp
--- a/Engine/Plugins/Animation/GameplayInsights/Source/RewindDebugger/Private/RewindDebuggerObjectTrack.cpp
+++ b/Engine/Plugins/Animation/GameplayInsights/Source/RewindDebugger/Private/RewindDebuggerObjectTrack.cpp
@@ -96,7 +96,7 @@ void FRewindDebuggerObjectTrack::IterateSubTracksInternal(TFunction<void(TShared
 		IteratorFunction(Track);
 	}
 };
-
+	
 FText FRewindDebuggerObjectTrack::GetDisplayNameInternal() const
 {
 	if (!bDisplayNameValid)
@@ -279,6 +279,15 @@ bool FRewindDebuggerObjectTrack::UpdateInternal()
 		}
 	}
 
+	if (bChanged)
+	{
+		Children.Sort([](const TSharedPtr<FRewindDebuggerTrack>& A, const TSharedPtr<FRewindDebuggerTrack>& B)
+		{
+			return A->GetSortOrderPriority() > B->GetSortOrderPriority() ||
+					A->GetDisplayName().ToString() < B->GetDisplayName().ToString();
+		});
+	}
+
 	for (auto& Child : Children)
 	{
 		if (Child->Update())
diff --git a/Engine/Plugins/Animation/GameplayInsights/Source/RewindDebugger/Private/RewindDebuggerObjectTrack.h b/Engine/Plugins/Animation/GameplayInsights/Source/RewindDebugger/Private/RewindDebuggerObjectTrack.h
--- a/Engine/Plugins/Animation/GameplayInsights/Source/RewindDebugger/Private/RewindDebuggerObjectTrack.h
+++ b/Engine/Plugins/Animation/GameplayInsights/Source/RewindDebugger/Private/RewindDebuggerObjectTrack.h
@@ -34,6 +34,7 @@ private:
 	virtual uint64 GetObjectIdInternal() const override { return ObjectId; }
 	virtual bool HasDebugDataInternal() const override { return false; }
 	virtual bool HandleDoubleClickInternal() override;
+	virtual int GetSortOrderPriorityInternal() const override { return -10; };
 	
 	mutable FText DisplayName;
 	FString ObjectName;
diff --git a/Engine/Plugins/Experimental/Animation/PoseSearch/Source/Editor/Private/PoseSearchDebugger.h b/Engine/Plugins/Experimental/Animation/PoseSearch/Source/Editor/Private/PoseSearchDebugger.h
--- a/Engine/Plugins/Experimental/Animation/PoseSearch/Source/Editor/Private/PoseSearchDebugger.h
+++ b/Engine/Plugins/Experimental/Animation/PoseSearch/Source/Editor/Private/PoseSearchDebugger.h
@@ -77,6 +77,7 @@ private:
 	virtual FName GetNameInternal() const override;
 	virtual FText GetDisplayNameInternal() const override;
 	virtual uint64 GetObjectIdInternal() const override { return ObjectId; }
+	virtual int GetSortOrderPriorityInternal() const override { return 10; };
 	virtual bool UpdateInternal() override;
 
 	TSharedPtr<SCostTimelineView> CostTimelineView;
diff --git a/Engine/Source/Editor/RewindDebuggerInterface/Public/RewindDebuggerTrack.h b/Engine/Source/Editor/RewindDebuggerInterface/Public/RewindDebuggerTrack.h
--- a/Engine/Source/Editor/RewindDebuggerInterface/Public/RewindDebuggerTrack.h
+++ b/Engine/Source/Editor/RewindDebuggerInterface/Public/RewindDebuggerTrack.h
@@ -59,6 +59,12 @@ public:
 	{
 		return GetNameInternal();
 	}
+	
+	// unique name for track (must match creator name if track is created by an IRewindDebuggerViewCreator) 
+	int GetSortOrderPriority() const
+	{
+		return GetSortOrderPriorityInternal();
+	}
 
 	// icon to display in the tree view
 	FSlateIcon GetIcon()
@@ -109,6 +115,7 @@ private:
 	virtual TSharedPtr<SWidget> GetDetailsViewInternal() { return TSharedPtr<SWidget>(); }
 	virtual FSlateIcon GetIconInternal() { return FSlateIcon(); }
 	virtual FName GetNameInternal() const { return ""; }
+	virtual int GetSortOrderPriorityInternal() const { return 0; }
 	virtual FText GetDisplayNameInternal() const { return FText(); }
 	virtual uint64 GetObjectIdInternal() const { return 0; }
 	virtual bool HasDebugDataInternal() const { return true; }
