Add names to rewind debugger extensions

diff --git a/Engine/Plugins/Animation/GameplayInsights/Source/RewindDebugger/Private/RewindDebugger.cpp b/Engine/Plugins/Animation/GameplayInsights/Source/RewindDebugger/Private/RewindDebugger.cpp
--- a/Engine/Plugins/Animation/GameplayInsights/Source/RewindDebugger/Private/RewindDebugger.cpp
+++ b/Engine/Plugins/Animation/GameplayInsights/Source/RewindDebugger/Private/RewindDebugger.cpp
@@ -1183,6 +1183,7 @@ void FRewindDebugger::Tick(float DeltaTime)
 		// update extensions
 		IterateExtensions([DeltaTime, this](IRewindDebuggerExtension* Extension)
 			{
+				TRACE_CPUPROFILER_EVENT_SCOPE_TEXT(*Extension->GetName());
 				Extension->Update(DeltaTime, this);
 			}
 		);
diff --git a/Engine/Plugins/Animation/GameplayInsights/Source/RewindDebugger/Private/RewindDebuggerAnimation.h b/Engine/Plugins/Animation/GameplayInsights/Source/RewindDebugger/Private/RewindDebuggerAnimation.h
--- a/Engine/Plugins/Animation/GameplayInsights/Source/RewindDebugger/Private/RewindDebuggerAnimation.h
+++ b/Engine/Plugins/Animation/GameplayInsights/Source/RewindDebugger/Private/RewindDebuggerAnimation.h
@@ -31,6 +31,7 @@ public:
 
 	virtual void Update(float DeltaTime, IRewindDebugger* RewindDebugger) override;
 	virtual void Clear(IRewindDebugger* RewindDebugger) override;
+	virtual FString GetName() { return TEXT("RewindDebuggerAnimation"); }
 	
 	void OnPIEResumed(bool bSimulating);
 	void OnPIEStopped(bool bSimulating);
diff --git a/Engine/Plugins/Animation/GameplayInsights/Source/RewindDebugger/Private/RewindDebuggerCamera.h b/Engine/Plugins/Animation/GameplayInsights/Source/RewindDebugger/Private/RewindDebuggerCamera.h
--- a/Engine/Plugins/Animation/GameplayInsights/Source/RewindDebugger/Private/RewindDebuggerCamera.h
+++ b/Engine/Plugins/Animation/GameplayInsights/Source/RewindDebugger/Private/RewindDebuggerCamera.h
@@ -22,6 +22,7 @@ public:
 	void Initialize();
 
 	virtual void Update(float DeltaTime, IRewindDebugger* RewindDebugger) override;
+	virtual FString GetName() { return TEXT("RewindDebuggerCamera"); }
 
 	ERewindDebuggerCameraMode CameraMode() const;
 	void SetCameraMode(ERewindDebuggerCameraMode Mode);
diff --git a/Engine/Plugins/Animation/GameplayInsights/Source/RewindDebuggerVLog/Private/RewindDebuggerVLog.h b/Engine/Plugins/Animation/GameplayInsights/Source/RewindDebuggerVLog/Private/RewindDebuggerVLog.h
--- a/Engine/Plugins/Animation/GameplayInsights/Source/RewindDebuggerVLog/Private/RewindDebuggerVLog.h
+++ b/Engine/Plugins/Animation/GameplayInsights/Source/RewindDebuggerVLog/Private/RewindDebuggerVLog.h
@@ -17,6 +17,8 @@ public:
 	void OnShowDebugInfo(UCanvas* Canvas, APlayerController* Player);
 	virtual ~FRewindDebuggerVLog();
 
+	virtual FString GetName() { return TEXT("FRewindDebuggerVLog"); }
+
 	void Initialize();
 	void MakeCategoriesMenu(UToolMenu* Menu);
 	void MakeLogLevelMenu(UToolMenu* Menu);
diff --git a/Engine/Plugins/Animation/PoseSearch/Source/Editor/Private/PoseSearchDebugger.h b/Engine/Plugins/Animation/PoseSearch/Source/Editor/Private/PoseSearchDebugger.h
--- a/Engine/Plugins/Animation/PoseSearch/Source/Editor/Private/PoseSearchDebugger.h
+++ b/Engine/Plugins/Animation/PoseSearch/Source/Editor/Private/PoseSearchDebugger.h
@@ -30,7 +30,7 @@ public:
 	static FDebugger* Get() { return Debugger; }
 	static void Initialize();
 	static void Shutdown();
-	static const FName ModularFeatureName;
+	virtual FString GetName() { return TEXT("PoseSearchDebugger"); }
 
 	// Shared data from the Rewind Debugger singleton
 	static bool IsPIESimulating();
diff --git a/Engine/Plugins/Chooser/Source/ChooserEditor/Private/RewindDebuggerChooser.h b/Engine/Plugins/Chooser/Source/ChooserEditor/Private/RewindDebuggerChooser.h
--- a/Engine/Plugins/Chooser/Source/ChooserEditor/Private/RewindDebuggerChooser.h
+++ b/Engine/Plugins/Chooser/Source/ChooserEditor/Private/RewindDebuggerChooser.h
@@ -14,4 +14,5 @@ public:
 	virtual ~FRewindDebuggerChooser() {};
 
 	virtual void Update(float DeltaTime, IRewindDebugger* RewindDebugger) override;
+	virtual FString GetName() { return TEXT("ChooserDebugger"); }
 };
diff --git a/Engine/Source/Editor/RewindDebuggerInterface/Public/IRewindDebuggerExtension.h b/Engine/Source/Editor/RewindDebuggerInterface/Public/IRewindDebuggerExtension.h
--- a/Engine/Source/Editor/RewindDebuggerInterface/Public/IRewindDebuggerExtension.h
+++ b/Engine/Source/Editor/RewindDebuggerInterface/Public/IRewindDebuggerExtension.h
@@ -18,7 +18,10 @@ class REWINDDEBUGGERINTERFACE_API IRewindDebuggerExtension : public IModularFeat
 public:
 	static const FName ModularFeatureName;
 
-	// called while scrubbing, playing back, or paused
+	// debugging name for this extension
+	virtual FString GetName() { return TEXT("RewindDebugger Extension"); }
+
+ 	// called while scrubbing, playing back, or paused
 	virtual void Update(float DeltaTime, IRewindDebugger* RewindDebugger) {};
 
 	// called when recording has started
