[RewindDebugger] added VisitSubTracks as an opt-in alternative to IterateSubTracks which can stop the iteration process. Default implementation is to call IterateSubTracks and Object track overrides it with proper recursive implementation.

diff --git a/Engine/Plugins/Animation/GameplayInsights/Source/RewindDebugger/Private/RewindDebuggerObjectTrack.cpp b/Engine/Plugins/Animation/GameplayInsights/Source/RewindDebugger/Private/RewindDebuggerObjectTrack.cpp
--- a/Engine/Plugins/Animation/GameplayInsights/Source/RewindDebugger/Private/RewindDebuggerObjectTrack.cpp
+++ b/Engine/Plugins/Animation/GameplayInsights/Source/RewindDebugger/Private/RewindDebuggerObjectTrack.cpp
@@ -181,6 +181,30 @@ void FRewindDebuggerObjectTrack::IterateSubTracksInternal(TFunction<void(TShared
 	}
 }
 
+FRewindDebuggerTrack::EVisitResult FRewindDebuggerObjectTrack::VisitSubTracksInternal(TFunction<EVisitResult(TSharedPtr<FRewindDebuggerTrack> SubTrack)> IteratorFunction)
+{
+	for (FTrackCreatorAndTrack& TrackChild : ChildTrackCreators)
+	{
+		if (TrackChild.Track.IsValid())
+		{
+			if (IteratorFunction(TrackChild.Track) == EVisitResult::Break)
+			{
+				return EVisitResult::Break;
+			}
+		}
+	}
+
+	for (const TSharedPtr<FRewindDebuggerTrack>& Track : Children)
+	{
+		if (Track->VisitSubTracks(IteratorFunction) == EVisitResult::Break)
+		{
+			return EVisitResult::Break;
+		}
+	}
+
+	return EVisitResult::Continue;
+}
+
 FName FRewindDebuggerObjectTrack::GetNameInternal() const
 {
 	if (PrimaryChildTrack.IsValid())
diff --git a/Engine/Plugins/Animation/GameplayInsights/Source/RewindDebugger/Private/RewindDebuggerObjectTrack.h b/Engine/Plugins/Animation/GameplayInsights/Source/RewindDebugger/Private/RewindDebuggerObjectTrack.h
--- a/Engine/Plugins/Animation/GameplayInsights/Source/RewindDebugger/Private/RewindDebuggerObjectTrack.h
+++ b/Engine/Plugins/Animation/GameplayInsights/Source/RewindDebugger/Private/RewindDebuggerObjectTrack.h
@@ -32,6 +32,7 @@ private:
 	virtual TSharedPtr<SWidget> GetTimelineViewInternal() override;
 	virtual bool UpdateInternal() override;
 	virtual void IterateSubTracksInternal(TFunction<void(TSharedPtr<FRewindDebuggerTrack> SubTrack)> IteratorFunction) override;
+	virtual EVisitResult VisitSubTracksInternal(TFunction<EVisitResult(TSharedPtr<FRewindDebuggerTrack> SubTrack)> IteratorFunction) override;
 	virtual FName GetNameInternal() const override;
 	virtual FSlateIcon GetIconInternal() override;
 	virtual FText GetDisplayNameInternal() const override;
diff --git a/Engine/Source/Editor/RewindDebuggerInterface/Public/RewindDebuggerTrack.h b/Engine/Source/Editor/RewindDebuggerInterface/Public/RewindDebuggerTrack.h
--- a/Engine/Source/Editor/RewindDebuggerInterface/Public/RewindDebuggerTrack.h
+++ b/Engine/Source/Editor/RewindDebuggerInterface/Public/RewindDebuggerTrack.h
@@ -30,6 +30,15 @@ ENUM_CLASS_FLAGS(ETrackHiddenFlags);
 class FRewindDebuggerTrack
 {
 public:
+	/** Result returned from a visitor functor indicating to continue or to quit early */
+	enum class EVisitResult : uint8
+	{
+		/** Stop iterating through tracks and early exit */
+		Break,
+		/**  Continue to iterate through all tracks */
+		Continue
+	};
+
 	FRewindDebuggerTrack()
 	{
 	}
@@ -128,7 +137,7 @@ public:
 	UE_DEPRECATED(5.7, "Use GetUObjectID instead.")
 	uint64 GetObjectId() const
 	{
-		return GetUObjectId ();
+		return GetUObjectId();
 	}
 
 	/** insights object id for an object associated with this track */
@@ -137,6 +146,12 @@ public:
 		return FObjectId(GetObjectIdInternal(), GetChildElementIdInternal());
 	}
 
+	/** iterates over all sub-tracks of this track and call Iterator function */
+	EVisitResult VisitSubTracks(TFunction<EVisitResult(TSharedPtr<FRewindDebuggerTrack> SubTrack)> IteratorFunction)
+	{
+		return VisitSubTracksInternal(IteratorFunction);
+	}
+
 	/** iterates over all sub-tracks of this track and call Iterator function */
 	void IterateSubTracks(TFunction<void(TSharedPtr<FRewindDebuggerTrack> SubTrack)> IteratorFunction)
 	{
@@ -242,6 +257,22 @@ private:
 	{
 	}
 
+	virtual EVisitResult VisitSubTracksInternal(TFunction<EVisitResult(TSharedPtr<FRewindDebuggerTrack> SubTrack)> IteratorFunction)
+	{
+		// This is an opt-in optimization added later in the API so default implementation is provided using IterateSubTracks
+		// Not perfect since we keep iterating, but we stop calling the function after getting a result that indicates to stop.
+		EVisitResult Result = EVisitResult::Continue;
+		IterateSubTracks([IteratorFunction, &Result](TSharedPtr<FRewindDebuggerTrack> SubTrack)
+			{
+				if (Result == EVisitResult::Continue)
+				{
+					Result = IteratorFunction(SubTrack);
+				}
+			});
+
+		return EVisitResult::Continue;
+	}
+
 	virtual bool HandleDoubleClickInternal()
 	{
 		IRewindDebugger::Instance()->OpenDetailsPanel();
