[RewindDebugger] added OnTrackListChanged callback to IRewindDebuggerExtension interface
Updated VisualLogger extension to use the new callback to avoid costly calls to IsObjectCurrentlyDebugger

diff --git a/Engine/Plugins/Animation/GameplayInsights/Source/RewindDebugger/Private/RewindDebugger.cpp b/Engine/Plugins/Animation/GameplayInsights/Source/RewindDebugger/Private/RewindDebugger.cpp
--- a/Engine/Plugins/Animation/GameplayInsights/Source/RewindDebugger/Private/RewindDebugger.cpp
+++ b/Engine/Plugins/Animation/GameplayInsights/Source/RewindDebugger/Private/RewindDebugger.cpp
@@ -361,7 +361,7 @@ void FRewindDebugger::RefreshDebugTracks()
 		{
 			DebugTracks[0] = MakeShared<FRewindDebuggerPlaceholderTrack>(FName(SelectionName), FText::FromString(SelectionName));
 			DebugTracks[1] = MakeShared<FRewindDebuggerPlaceholderTrack>(DebugMessageTrackName, NSLOCTEXT("RewindDebugger", "No Debug Data", " - Start a recording to debug"));
-			TrackListChangedDelegate.ExecuteIfBound();
+			OnTrackListChanged();
 		}
 	}
 	else if (DebugTracks.Num() || CandidateIds.Num())
@@ -398,13 +398,14 @@ void FRewindDebugger::RefreshDebugTracks()
 		{
 			if (DebugTrack->Update())
 			{
+				UE_LOG(LogRewindDebugger, Verbose, TEXT("List changed by: '%s'"), *DebugTrack->GetDisplayName().ToString());
 				bChanged = true;
 			}
 		}
 
 		if (bChanged)
 		{
-			TrackListChangedDelegate.ExecuteIfBound();
+			OnTrackListChanged();
 		}
 	}
 }
@@ -1159,6 +1160,18 @@ void FRewindDebugger::Tick(float DeltaTime)
 	}
 }
 
+void FRewindDebugger::OnTrackListChanged()
+{
+	TrackListChangedDelegate.ExecuteIfBound();
+
+	IterateExtensions([this](IRewindDebuggerExtension* Extension)
+		{
+			TRACE_CPUPROFILER_EVENT_SCOPE_TEXT(*Extension->GetName());
+			Extension->OnTrackListChanged(this);
+		}
+	);
+}
+
 void FRewindDebugger::OpenDetailsPanel()
 {
 	bIsDetailsPanelOpen = true;
diff --git a/Engine/Plugins/Animation/GameplayInsights/Source/RewindDebugger/Private/RewindDebugger.h b/Engine/Plugins/Animation/GameplayInsights/Source/RewindDebugger/Private/RewindDebugger.h
--- a/Engine/Plugins/Animation/GameplayInsights/Source/RewindDebugger/Private/RewindDebugger.h
+++ b/Engine/Plugins/Animation/GameplayInsights/Source/RewindDebugger/Private/RewindDebugger.h
@@ -220,6 +220,8 @@ public:
 private:
 	static void RefreshDebuggedObjects(TArray<TSharedPtr<RewindDebugger::FRewindDebuggerTrack>>& InTracks, TArray<TSharedPtr<FDebugObjectInfo>>& OutObjects);
 
+	void OnTrackListChanged();
+
 	void OnPIEStarted(bool bSimulating);
 	void OnPIEPaused(bool bSimulating);
 	void OnPIEResumed(bool bSimulating);
diff --git a/Engine/Plugins/Animation/GameplayInsights/Source/RewindDebuggerVLog/Private/RewindDebuggerVLog.cpp b/Engine/Plugins/Animation/GameplayInsights/Source/RewindDebuggerVLog/Private/RewindDebuggerVLog.cpp
--- a/Engine/Plugins/Animation/GameplayInsights/Source/RewindDebuggerVLog/Private/RewindDebuggerVLog.cpp
+++ b/Engine/Plugins/Animation/GameplayInsights/Source/RewindDebuggerVLog/Private/RewindDebuggerVLog.cpp
@@ -125,13 +125,10 @@ void FRewindDebuggerVLog::RenderLogEntry(const FVisualLogEntry& Entry, UCanvas*
 void FRewindDebuggerVLog::ImmediateRender(const UObject* Object, const FVisualLogEntry& Entry)
 {
 #if OBJECT_TRACE_ENABLED
-	if (IRewindDebugger* RewindDebugger = IRewindDebugger::Instance())
+	uint64 ObjectId = FObjectTrace::GetObjectId(Object);
+	if (DebuggedObjectIds.Contains(ObjectId))
 	{
-		uint64 ObjectId = FObjectTrace::GetObjectId(Object);
-		if (RewindDebugger->IsObjectCurrentlyDebugged(ObjectId))
-		{
-			ImmediateRenderQueue.Add(Entry);
-		}
+		ImmediateRenderQueue.Add(Entry);
 	}
 #endif
 }
@@ -253,4 +250,17 @@ void FRewindDebuggerVLog::Update(float DeltaTime, IRewindDebugger* RewindDebugge
 {
 }
 
-#undef LOCTEXT_NAMESPACE
\ No newline at end of file
+void FRewindDebuggerVLog::OnTrackListChanged(IRewindDebugger* RewindDebugger)
+{
+	const TArray<TSharedPtr<FDebugObjectInfo>>& DebuggedObjects = RewindDebugger->GetDebuggedObjects();
+	DebuggedObjectIds.Reset(DebuggedObjects.Num());
+	for (const TSharedPtr<FDebugObjectInfo>& Object : DebuggedObjects)
+	{
+		if (Object->Id.IsSet())
+		{
+			DebuggedObjectIds.Push(Object->GetUObjectId());
+		}
+	}
+}
+
+#undef LOCTEXT_NAMESPACE
diff --git a/Engine/Plugins/Animation/GameplayInsights/Source/RewindDebuggerVLog/Private/RewindDebuggerVLog.h b/Engine/Plugins/Animation/GameplayInsights/Source/RewindDebuggerVLog/Private/RewindDebuggerVLog.h
--- a/Engine/Plugins/Animation/GameplayInsights/Source/RewindDebuggerVLog/Private/RewindDebuggerVLog.h
+++ b/Engine/Plugins/Animation/GameplayInsights/Source/RewindDebuggerVLog/Private/RewindDebuggerVLog.h
@@ -30,7 +30,9 @@ public:
 	static ELogVerbosity::Type GetMinLogVerbosity();
 	static void SetMinLogVerbosity(ELogVerbosity::Type Value);
 
+protected:
 	virtual void Update(float DeltaTime, IRewindDebugger* RewindDebugger) override;
+	virtual void OnTrackListChanged(IRewindDebugger* RewindDebugger) override;
 
 private:
 	void AddLogEntries(const TArray<TSharedPtr<FDebugObjectInfo>>& Components, float StartTime, float EndTime, const class IVisualLoggerProvider* Provider, UCanvas* Canvas);
@@ -47,5 +49,6 @@ private:
 	FDelegateHandle DelegateHandle;
 	UFont* MonospaceFont = nullptr;
 
+	TArray<uint64> DebuggedObjectIds;
 	TArray<FVisualLogEntry> ImmediateRenderQueue;
 };
diff --git a/Engine/Source/Editor/RewindDebuggerInterface/Public/IRewindDebuggerExtension.h b/Engine/Source/Editor/RewindDebuggerInterface/Public/IRewindDebuggerExtension.h
--- a/Engine/Source/Editor/RewindDebuggerInterface/Public/IRewindDebuggerExtension.h
+++ b/Engine/Source/Editor/RewindDebuggerInterface/Public/IRewindDebuggerExtension.h
@@ -7,7 +7,7 @@
 #define UE_API REWINDDEBUGGERINTERFACE_API
 
 
-// IRewindDebuggerExtension 
+// IRewindDebuggerExtension
 //
 // interface class for extensions which add functionality to the rewind debugger
 // these will get updated on scrub/playback to handle things like updating the world state to match recorded data from that time for a particular system
@@ -20,20 +20,36 @@ class IRewindDebuggerExtension : public IModularFeature
 public:
 	static UE_API const FName ModularFeatureName;
 
-	// debugging name for this extension
-	virtual FString GetName() { return TEXT("RewindDebugger Extension"); }
-
- 	// called while scrubbing, playing back, or paused
-	virtual void Update(float DeltaTime, IRewindDebugger* RewindDebugger) {};
-
-	// called when recording has started
-	virtual void RecordingStarted(IRewindDebugger* RewindDebugger) {};
-
-	// called when recording has ended
-	virtual void RecordingStopped(IRewindDebugger* RewindDebugger) {};
-	
-	// called when recording is unloaded
-	virtual void Clear(IRewindDebugger* RewindDebugger) {};
+	/** debugging name for this extension */
+	virtual FString GetName()
+	{
+		return TEXT("RewindDebugger Extension");
+	}
+
+	/** called while scrubbing, playing back, or paused */
+	virtual void Update(float DeltaTime, IRewindDebugger* RewindDebugger)
+	{
+	};
+
+	/** called when recording has started */
+	virtual void RecordingStarted(IRewindDebugger* RewindDebugger)
+	{
+	};
+
+	/** called when recording has ended */
+	virtual void RecordingStopped(IRewindDebugger* RewindDebugger)
+	{
+	};
+
+	/** called when recording is unloaded */
+	virtual void Clear(IRewindDebugger* RewindDebugger)
+	{
+	};
+
+	// called when tracks get removed or added to the debugger
+	virtual void OnTrackListChanged(IRewindDebugger* RewindDebugger)
+	{
+	};
 };
 
 #undef UE_API
