[RewindDebugger] deprecated header and update comment style
#rnx

diff --git a/Engine/Source/Editor/RewindDebuggerInterface/Public/IRewindDebuggerExtension.h b/Engine/Source/Editor/RewindDebuggerInterface/Public/IRewindDebuggerExtension.h
--- a/Engine/Source/Editor/RewindDebuggerInterface/Public/IRewindDebuggerExtension.h
+++ b/Engine/Source/Editor/RewindDebuggerInterface/Public/IRewindDebuggerExtension.h
@@ -1,20 +1,20 @@
 // Copyright Epic Games, Inc. All Rights Reserved.
 #pragma once
 
+#if UE_ENABLE_INCLUDE_ORDER_DEPRECATED_IN_5_7
 #include "CoreMinimal.h"
+#endif // UE_ENABLE_INCLUDE_ORDER_DEPRECATED_IN_5_7
 #include "Features/IModularFeatures.h"
 
 #define UE_API REWINDDEBUGGERINTERFACE_API
 
-
-// IRewindDebuggerExtension
-//
-// interface class for extensions which add functionality to the rewind debugger
-// these will get updated on scrub/playback to handle things like updating the world state to match recorded data from that time for a particular system
-//
-
 class IRewindDebugger;
 
+/**
+ * Interface class for extensions which add functionality to the rewind debugger.
+ * These will get updated on scrub/playback to handle things like updating the world state
+ * to match recorded data from that time for a particular system.
+ */
 class IRewindDebuggerExtension : public IModularFeature
 {
 public:
