[RewindDebugger] fixed CIS by using deprecation macro properly

#Synced-CL 49424914

diff --git a/Engine/Source/Editor/RewindDebuggerInterface/Public/IRewindDebuggerExtension.h b/Engine/Source/Editor/RewindDebuggerInterface/Public/IRewindDebuggerExtension.h
--- a/Engine/Source/Editor/RewindDebuggerInterface/Public/IRewindDebuggerExtension.h
+++ b/Engine/Source/Editor/RewindDebuggerInterface/Public/IRewindDebuggerExtension.h
@@ -1,4 +1,4 @@
-﻿// Copyright Epic Games, Inc. All Rights Reserved.
+// Copyright Epic Games, Inc. All Rights Reserved.
 #pragma once
 
 #if UE_ENABLE_INCLUDE_ORDER_DEPRECATED_IN_5_7
@@ -41,12 +41,12 @@ public:
 	{
 	}
 
-	UE_DEPRECATED(5.8, TEXT("Use AnalysisSessionOpened instead or implement an extension inheriting from IRewindDebuggerRuntimeExtension"))
+	UE_DEPRECATED(5.8, "Use AnalysisSessionOpened instead or implement an extension inheriting from IRewindDebuggerRuntimeExtension")
 	virtual void RecordingStarted(IRewindDebugger* RewindDebugger) final
 	{
 	};
 
-	UE_DEPRECATED(5.8, TEXT("Use AnalysisSessionClosed instead or implement an extension inheriting from IRewindDebuggerRuntimeExtension"))
+	UE_DEPRECATED(5.8, "Use AnalysisSessionClosed instead or implement an extension inheriting from IRewindDebuggerRuntimeExtension")
 	virtual void RecordingStopped(IRewindDebugger* RewindDebugger) final
 	{
 	};
